Fix Mac & copyright notice CIS issues

diff --git a/Engine/Plugins/Experimental/WidgetEditorToolPalette/Source/WidgetEditorToolPalette/Public/Settings/CreateWidgetToolSettings.h b/Engine/Plugins/Experimental/WidgetEditorToolPalette/Source/WidgetEditorToolPalette/Public/Settings/CreateWidgetToolSettings.h
--- a/Engine/Plugins/Experimental/WidgetEditorToolPalette/Source/WidgetEditorToolPalette/Public/Settings/CreateWidgetToolSettings.h
+++ b/Engine/Plugins/Experimental/WidgetEditorToolPalette/Source/WidgetEditorToolPalette/Public/Settings/CreateWidgetToolSettings.h
@@ -1,4 +1,4 @@
-// Copyright Epic Games, Inc. All Rghts Reserved.
+// Copyright Epic Games, Inc. All Rights Reserved.
 
 #pragma once
 
diff --git a/Engine/Source/Editor/UMGEditor/UMGEditor.Build.cs b/Engine/Source/Editor/UMGEditor/UMGEditor.Build.cs
--- a/Engine/Source/Editor/UMGEditor/UMGEditor.Build.cs
+++ b/Engine/Source/Editor/UMGEditor/UMGEditor.Build.cs
@@ -47,6 +47,7 @@ public class UMGEditor : ModuleRules
 				"InputCore",
 				"Engine",
 				"AssetTools",
+				"EditorSubsystem",
 				"EditorFramework",
 				"InteractiveToolsFramework",
 				"UnrealEd", // for Asset Editor Subsystem
diff --git a/Engine/Source/Runtime/InteractiveToolsFramework/Private/BaseBehaviors/Widgets/WidgetBaseBehavior.cpp b/Engine/Source/Runtime/InteractiveToolsFramework/Private/BaseBehaviors/Widgets/WidgetBaseBehavior.cpp
--- a/Engine/Source/Runtime/InteractiveToolsFramework/Private/BaseBehaviors/Widgets/WidgetBaseBehavior.cpp
+++ b/Engine/Source/Runtime/InteractiveToolsFramework/Private/BaseBehaviors/Widgets/WidgetBaseBehavior.cpp
@@ -1,3 +1,5 @@
+// Copyright Epic Games, Inc. All Rights Reserved.
+
 #include "BaseBehaviors/Widgets/WidgetBaseBehavior.h"
 
 bool IWidgetBaseBehavior::OnKeyChar(const FGeometry& MyGeometry, const FCharacterEvent& InCharacterEvent)
