Remove WhitelistPlatforms usages.

diff --git a/Engine/Plugins/Experimental/WidgetEditorToolPalette/WidgetEditorToolPalette.uplugin b/Engine/Plugins/Experimental/WidgetEditorToolPalette/WidgetEditorToolPalette.uplugin
--- a/Engine/Plugins/Experimental/WidgetEditorToolPalette/WidgetEditorToolPalette.uplugin
+++ b/Engine/Plugins/Experimental/WidgetEditorToolPalette/WidgetEditorToolPalette.uplugin
@@ -19,7 +19,7 @@
 			"Name": "WidgetEditorToolPalette",
 			"Type": "Editor",
 			"LoadingPhase": "Default",
-			"WhitelistPlatforms": [ "Win64" ]
+			"PlatformAllowList": [ "Win64" ]
 		}
 	]
 }
