Fix default stacks not populating in WidgetEditorToolPalette

diff --git a/Engine/Plugins/Experimental/WidgetEditorToolPalette/Source/WidgetEditorToolPalette/Public/Settings/CreateWidgetToolSettings.h b/Engine/Plugins/Experimental/WidgetEditorToolPalette/Source/WidgetEditorToolPalette/Public/Settings/CreateWidgetToolSettings.h
--- a/Engine/Plugins/Experimental/WidgetEditorToolPalette/Source/WidgetEditorToolPalette/Public/Settings/CreateWidgetToolSettings.h
+++ b/Engine/Plugins/Experimental/WidgetEditorToolPalette/Source/WidgetEditorToolPalette/Public/Settings/CreateWidgetToolSettings.h
@@ -61,7 +61,7 @@ struct WIDGETEDITORTOOLPALETTE_API FCreateWidgetStackInfo
 /**
  * Allows for create widget tool layout to be defined via settings.
  */
-UCLASS(config=EditorPerProjectUserSettings, autoexpandcategories=(Keybinds))
+UCLASS(config=WidgetEditorToolPalette, autoexpandcategories=(Keybinds))
 class WIDGETEDITORTOOLPALETTE_API UCreateWidgetToolSettings : public UDeveloperSettings
 {
 	GENERATED_BODY()
